Break XInput out into it's own Input Device module instead of being created as a part of the Windows Application.

This will pave the way for us to enable GameInput as the default input device in the future, and make it much easier to just turn off XInput via the plugin system instead of needing to add a bunch of checks to the WindowsApplication.cpp

Developer Note: This means that if you for some reason modified the engine to make XInputInterface.h a publicly available include, you need to add the XInputDevice module as a dependency. This should not be the case for anyone, unless they have customized the engine to extend on XInput and make it public in the first place (It was in a private module folder and not exposed with an API, so we can't really do anything about it if it was modified).

diff --git a/Engine/Plugins/Runtime/Windows/XInputDevice/Source/XInputDevice/Private/XInputDeviceModule.cpp b/Engine/Plugins/Runtime/Windows/XInputDevice/Source/XInputDevice/Private/XInputDeviceModule.cpp
new file mode 100644
--- /dev/null
+++ b/Engine/Plugins/Runtime/Windows/XInputDevice/Source/XInputDevice/Private/XInputDeviceModule.cpp
@@ -0,0 +1,12 @@
+// Copyright Epic Games, Inc. All Rights Reserved.
+
+#include "XInputDeviceModule.h"
+#include "XInputInterface.h"
+#include "Modules/ModuleManager.h"	// For IMPLEMENT_MODULE
+
+TSharedPtr<IInputDevice> FXInputDeviceModule::CreateInputDevice(const TSharedRef<FGenericApplicationMessageHandler>& InMessageHandler)
+{
+	return XInputInterface::Create(InMessageHandler);
+}
+
+IMPLEMENT_MODULE(FXInputDeviceModule, XInputDevice)
\ No newline at end of file
diff --git a/Engine/Plugins/Runtime/Windows/XInputDevice/Source/XInputDevice/Private/XInputInterface.cpp b/Engine/Plugins/Runtime/Windows/XInputDevice/Source/XInputDevice/Private/XInputInterface.cpp
new file mode 100644
--- /dev/null
+++ b/Engine/Plugins/Runtime/Windows/XInputDevice/Source/XInputDevice/Private/XInputInterface.cpp
@@ -0,0 +1,329 @@
+// Copyright Epic Games, Inc. All Rights Reserved.
+
+#include "XInputInterface.h"
+#include "HAL/PlatformTime.h"
+#include "Misc/CoreDelegates.h"
+#include "Windows/WindowsApplication.h"
+#include "GenericPlatform/GenericApplication.h"
+#include "GenericPlatform/GenericPlatformInputDeviceMapper.h"
+#include "Misc/ConfigCacheIni.h"
+
+#pragma pack (push,8)
+#include "Windows/AllowWindowsPlatformTypes.h"
+#include <XInput.h>
+#include "Windows/HideWindowsPlatformTypes.h"
+#pragma pack (pop)
+
+static int32 ForceControllerStateUpdate = 0;
+FAutoConsoleVariableRef CVarForceControllerStateUpdate(
+	TEXT("XInput.ForceControllerStateUpdate"),
+	ForceControllerStateUpdate,
+	TEXT("Force XInput refresh of controller state on each frame.\n")
+	TEXT("0: Not Enabled, 1: Enabled"),
+	ECVF_Default);
+
+TSharedRef< XInputInterface > XInputInterface::Create(  const TSharedRef< FGenericApplicationMessageHandler >& InMessageHandler )
+{
+	return MakeShareable( new XInputInterface( InMessageHandler ) );
+}
+
+
+XInputInterface::XInputInterface( const TSharedRef< FGenericApplicationMessageHandler >& InMessageHandler )
+	: MessageHandler( InMessageHandler )
+{
+	for ( int32 ControllerIndex=0; ControllerIndex < MAX_NUM_XINPUT_CONTROLLERS; ++ControllerIndex )
+	{
+		FControllerState& ControllerState = ControllerStates[ControllerIndex];
+		FMemory::Memzero( &ControllerState, sizeof(FControllerState) );
+
+		ControllerState.ControllerId = ControllerIndex;
+	}
+
+	bIsGamepadAttached = false;
+	bNeedsControllerStateUpdate = true;
+	InitialButtonRepeatDelay = 0.2f;
+	ButtonRepeatDelay = 0.1f;
+
+	GConfig->GetFloat(TEXT("/Script/Engine.InputSettings"), TEXT("InitialButtonRepeatDelay"), InitialButtonRepeatDelay, GInputIni);
+	GConfig->GetFloat(TEXT("/Script/Engine.InputSettings"), TEXT("ButtonRepeatDelay"), ButtonRepeatDelay, GInputIni);
+
+	// In the engine, all controllers map to xbox controllers for consistency 
+	X360ToXboxControllerMapping[0] = 0;		// A
+	X360ToXboxControllerMapping[1] = 1;		// B
+	X360ToXboxControllerMapping[2] = 2;		// X
+	X360ToXboxControllerMapping[3] = 3;		// Y
+	X360ToXboxControllerMapping[4] = 4;		// L1
+	X360ToXboxControllerMapping[5] = 5;		// R1
+	X360ToXboxControllerMapping[6] = 7;		// Back 
+	X360ToXboxControllerMapping[7] = 6;		// Start
+	X360ToXboxControllerMapping[8] = 8;		// Left thumbstick
+	X360ToXboxControllerMapping[9] = 9;		// Right thumbstick
+	X360ToXboxControllerMapping[10] = 10;	// L2
+	X360ToXboxControllerMapping[11] = 11;	// R2
+	X360ToXboxControllerMapping[12] = 12;	// Dpad up
+	X360ToXboxControllerMapping[13] = 13;	// Dpad down
+	X360ToXboxControllerMapping[14] = 14;	// Dpad left
+	X360ToXboxControllerMapping[15] = 15;	// Dpad right
+	X360ToXboxControllerMapping[16] = 16;	// Left stick up
+	X360ToXboxControllerMapping[17] = 17;	// Left stick down
+	X360ToXboxControllerMapping[18] = 18;	// Left stick left
+	X360ToXboxControllerMapping[19] = 19;	// Left stick right
+	X360ToXboxControllerMapping[20] = 20;	// Right stick up
+	X360ToXboxControllerMapping[21] = 21;	// Right stick down
+	X360ToXboxControllerMapping[22] = 22;	// Right stick left
+	X360ToXboxControllerMapping[23] = 23;	// Right stick right
+
+	Buttons[0] = FGamepadKeyNames::FaceButtonBottom;
+	Buttons[1] = FGamepadKeyNames::FaceButtonRight;
+	Buttons[2] = FGamepadKeyNames::FaceButtonLeft;
+	Buttons[3] = FGamepadKeyNames::FaceButtonTop;
+	Buttons[4] = FGamepadKeyNames::LeftShoulder;
+	Buttons[5] = FGamepadKeyNames::RightShoulder;
+	Buttons[6] = FGamepadKeyNames::SpecialRight;
+	Buttons[7] = FGamepadKeyNames::SpecialLeft;
+	Buttons[8] = FGamepadKeyNames::LeftThumb;
+	Buttons[9] = FGamepadKeyNames::RightThumb;
+	Buttons[10] = FGamepadKeyNames::LeftTriggerThreshold;
+	Buttons[11] = FGamepadKeyNames::RightTriggerThreshold;
+	Buttons[12] = FGamepadKeyNames::DPadUp;
+	Buttons[13] = FGamepadKeyNames::DPadDown;
+	Buttons[14] = FGamepadKeyNames::DPadLeft;
+	Buttons[15] = FGamepadKeyNames::DPadRight;
+	Buttons[16] = FGamepadKeyNames::LeftStickUp;
+	Buttons[17] = FGamepadKeyNames::LeftStickDown;
+	Buttons[18] = FGamepadKeyNames::LeftStickLeft;
+	Buttons[19] = FGamepadKeyNames::LeftStickRight;
+	Buttons[20] = FGamepadKeyNames::RightStickUp;
+	Buttons[21] = FGamepadKeyNames::RightStickDown;
+	Buttons[22] = FGamepadKeyNames::RightStickLeft;
+	Buttons[23] = FGamepadKeyNames::RightStickRight;
+}
+
+float ShortToNormalizedFloat(int16 AxisVal)
+{
+	// normalize [-32768..32767] -> [-1..1]
+	const float Norm = (AxisVal <= 0 ? 32768.f : 32767.f);
+	return float(AxisVal) / Norm;
+}
+
+static FName XInputInterfaceName = FName("XInputInterface");
+static FString XInputControllerIdentifier = TEXT("XInputController");
+
+void XInputInterface::SendControllerEvents()
+{
+	bool bWereConnected[MAX_NUM_XINPUT_CONTROLLERS];
+	XINPUT_STATE XInputStates[MAX_NUM_XINPUT_CONTROLLERS];
+
+	bIsGamepadAttached = false;
+	for ( int32 ControllerIndex=0; ControllerIndex < MAX_NUM_XINPUT_CONTROLLERS; ++ControllerIndex )
+	{
+		FControllerState& ControllerState = ControllerStates[ControllerIndex];
+
+		bWereConnected[ControllerIndex] = ControllerState.bIsConnected;
+
+		if (ControllerState.bIsConnected || bNeedsControllerStateUpdate || ForceControllerStateUpdate != 0)
+		{
+			XINPUT_STATE& XInputState = XInputStates[ControllerIndex];
+			FMemory::Memzero( &XInputState, sizeof(XINPUT_STATE) );
+
+			ControllerState.bIsConnected = ( XInputGetState( ControllerIndex, &XInputState ) == ERROR_SUCCESS ) ? true : false;
+
+			if (ControllerState.bIsConnected)
+			{
+				bIsGamepadAttached = true;
+			}
+		}
+	}
+		
+	for ( int32 ControllerIndex = 0; ControllerIndex < MAX_NUM_XINPUT_CONTROLLERS; ++ControllerIndex )
+	{
+		// Set input scope, there doesn't seem to be a reliable way to differentiate 360 vs Xbox one controllers so use generic name
+		FInputDeviceScope InputScope(this, XInputInterfaceName, ControllerIndex, XInputControllerIdentifier);
+
+		FControllerState& ControllerState = ControllerStates[ControllerIndex];
+
+		const bool bWasConnected = bWereConnected[ControllerIndex];
+
+		// If the controller is connected send events or if the controller was connected send a final event with default states so that 
+		// the game doesn't think that controller buttons are still held down
+		if( ControllerState.bIsConnected || bWasConnected )
+		{
+			const XINPUT_STATE& XInputState = XInputStates[ControllerIndex];
+
+			IPlatformInputDeviceMapper& DeviceMapper = IPlatformInputDeviceMapper::Get();
+			FPlatformUserId PlatformUser = PLATFORMUSERID_NONE;
+			FInputDeviceId InputDevice = INPUTDEVICEID_NONE;
+			DeviceMapper.RemapControllerIdToPlatformUserAndDevice(ControllerState.ControllerId, OUT PlatformUser, OUT InputDevice);
+
+			// If the controller is connected now but was not before, refresh the information
+			if (!bWasConnected && ControllerState.bIsConnected)
+			{
+				DeviceMapper.Internal_MapInputDeviceToUser(InputDevice, PlatformUser, EInputDeviceConnectionState::Connected);
+			}
+			else if (bWasConnected && !ControllerState.bIsConnected)
+			{
+				DeviceMapper.Internal_MapInputDeviceToUser(InputDevice, PlatformUser, EInputDeviceConnectionState::Disconnected);
+			}
+			
+			bool CurrentStates[MAX_NUM_CONTROLLER_BUTTONS] = {0};
+		
+			// Get the current state of all buttons
+			CurrentStates[X360ToXboxControllerMapping[0]] = !!(XInputState.Gamepad.wButtons & XINPUT_GAMEPAD_A);
+			CurrentStates[X360ToXboxControllerMapping[1]] = !!(XInputState.Gamepad.wButtons & XINPUT_GAMEPAD_B);
+			CurrentStates[X360ToXboxControllerMapping[2]] = !!(XInputState.Gamepad.wButtons & XINPUT_GAMEPAD_X);
+			CurrentStates[X360ToXboxControllerMapping[3]] = !!(XInputState.Gamepad.wButtons & XINPUT_GAMEPAD_Y);
+			CurrentStates[X360ToXboxControllerMapping[4]] = !!(XInputState.Gamepad.wButtons & XINPUT_GAMEPAD_LEFT_SHOULDER);
+			CurrentStates[X360ToXboxControllerMapping[5]] = !!(XInputState.Gamepad.wButtons & XINPUT_GAMEPAD_RIGHT_SHOULDER);
+			CurrentStates[X360ToXboxControllerMapping[6]] = !!(XInputState.Gamepad.wButtons & XINPUT_GAMEPAD_BACK);
+			CurrentStates[X360ToXboxControllerMapping[7]] = !!(XInputState.Gamepad.wButtons & XINPUT_GAMEPAD_START);
+			CurrentStates[X360ToXboxControllerMapping[8]] = !!(XInputState.Gamepad.wButtons & XINPUT_GAMEPAD_LEFT_THUMB);
+			CurrentStates[X360ToXboxControllerMapping[9]] = !!(XInputState.Gamepad.wButtons & XINPUT_GAMEPAD_RIGHT_THUMB);
+			CurrentStates[X360ToXboxControllerMapping[10]] = !!(XInputState.Gamepad.bLeftTrigger > XINPUT_GAMEPAD_TRIGGER_THRESHOLD);
+			CurrentStates[X360ToXboxControllerMapping[11]] = !!(XInputState.Gamepad.bRightTrigger > XINPUT_GAMEPAD_TRIGGER_THRESHOLD);
+			CurrentStates[X360ToXboxControllerMapping[12]] = !!(XInputState.Gamepad.wButtons & XINPUT_GAMEPAD_DPAD_UP);
+			CurrentStates[X360ToXboxControllerMapping[13]] = !!(XInputState.Gamepad.wButtons & XINPUT_GAMEPAD_DPAD_DOWN);
+			CurrentStates[X360ToXboxControllerMapping[14]] = !!(XInputState.Gamepad.wButtons & XINPUT_GAMEPAD_DPAD_LEFT);
+			CurrentStates[X360ToXboxControllerMapping[15]] = !!(XInputState.Gamepad.wButtons & XINPUT_GAMEPAD_DPAD_RIGHT);
+			CurrentStates[X360ToXboxControllerMapping[16]] = !!(XInputState.Gamepad.sThumbLY > XINPUT_GAMEPAD_LEFT_THUMB_DEADZONE);
+			CurrentStates[X360ToXboxControllerMapping[17]] = !!(XInputState.Gamepad.sThumbLY < -XINPUT_GAMEPAD_LEFT_THUMB_DEADZONE);
+			CurrentStates[X360ToXboxControllerMapping[18]] = !!(XInputState.Gamepad.sThumbLX < -XINPUT_GAMEPAD_LEFT_THUMB_DEADZONE);
+			CurrentStates[X360ToXboxControllerMapping[19]] = !!(XInputState.Gamepad.sThumbLX > XINPUT_GAMEPAD_LEFT_THUMB_DEADZONE);
+			CurrentStates[X360ToXboxControllerMapping[20]] = !!(XInputState.Gamepad.sThumbRY > XINPUT_GAMEPAD_RIGHT_THUMB_DEADZONE);
+			CurrentStates[X360ToXboxControllerMapping[21]] = !!(XInputState.Gamepad.sThumbRY < -XINPUT_GAMEPAD_RIGHT_THUMB_DEADZONE);
+			CurrentStates[X360ToXboxControllerMapping[22]] = !!(XInputState.Gamepad.sThumbRX < -XINPUT_GAMEPAD_RIGHT_THUMB_DEADZONE);
+			CurrentStates[X360ToXboxControllerMapping[23]] = !!(XInputState.Gamepad.sThumbRX > XINPUT_GAMEPAD_RIGHT_THUMB_DEADZONE);
+
+			// Send new analog data if it's different or outside the platform deadzone.
+			auto OnControllerAnalog = [this, &PlatformUser, &InputDevice](const FName& GamePadKey, const auto NewAxisValue, const float NewAxisValueNormalized, auto& OldAxisValue, const auto DeadZone) {
+				if (OldAxisValue != NewAxisValue || FMath::Abs((int32)NewAxisValue) > DeadZone)
+				{
+					MessageHandler->OnControllerAnalog(GamePadKey, PlatformUser, InputDevice, NewAxisValueNormalized);
+				}
+				OldAxisValue = NewAxisValue;
+			};
+
+			const auto& Gamepad = XInputState.Gamepad;
+
+			OnControllerAnalog(FGamepadKeyNames::LeftAnalogX, Gamepad.sThumbLX, ShortToNormalizedFloat(Gamepad.sThumbLX), ControllerState.LeftXAnalog, XINPUT_GAMEPAD_LEFT_THUMB_DEADZONE);
+			OnControllerAnalog(FGamepadKeyNames::LeftAnalogY, Gamepad.sThumbLY, ShortToNormalizedFloat(Gamepad.sThumbLY), ControllerState.LeftYAnalog, XINPUT_GAMEPAD_LEFT_THUMB_DEADZONE);
+
+			OnControllerAnalog(FGamepadKeyNames::RightAnalogX, Gamepad.sThumbRX, ShortToNormalizedFloat(Gamepad.sThumbRX), ControllerState.RightXAnalog, XINPUT_GAMEPAD_RIGHT_THUMB_DEADZONE);
+			OnControllerAnalog(FGamepadKeyNames::RightAnalogY, Gamepad.sThumbRY, ShortToNormalizedFloat(Gamepad.sThumbRY), ControllerState.RightYAnalog, XINPUT_GAMEPAD_RIGHT_THUMB_DEADZONE);
+
+			OnControllerAnalog(FGamepadKeyNames::LeftTriggerAnalog, Gamepad.bLeftTrigger, Gamepad.bLeftTrigger / 255.f, ControllerState.LeftTriggerAnalog, XINPUT_GAMEPAD_TRIGGER_THRESHOLD);
+			OnControllerAnalog(FGamepadKeyNames::RightTriggerAnalog, Gamepad.bRightTrigger, Gamepad.bRightTrigger / 255.f, ControllerState.RightTriggerAnalog, XINPUT_GAMEPAD_TRIGGER_THRESHOLD);
+
+			const double CurrentTime = FPlatformTime::Seconds();
+
+			// For each button check against the previous state and send the correct message if any
+			for (int32 ButtonIndex = 0; ButtonIndex < MAX_NUM_CONTROLLER_BUTTONS; ++ButtonIndex)
+			{
+				if (CurrentStates[ButtonIndex] != ControllerState.ButtonStates[ButtonIndex])
+				{
+					if( CurrentStates[ButtonIndex] )
+					{
+						MessageHandler->OnControllerButtonPressed( Buttons[ButtonIndex], PlatformUser, InputDevice, false );
+					}
+					else
+					{
+						MessageHandler->OnControllerButtonReleased( Buttons[ButtonIndex], PlatformUser, InputDevice, false );
+					}
+
+					if ( CurrentStates[ButtonIndex] != 0 )
+					{
+						// this button was pressed - set the button's NextRepeatTime to the InitialButtonRepeatDelay
+						ControllerState.NextRepeatTime[ButtonIndex] = CurrentTime + InitialButtonRepeatDelay;
+					}
+				}
+				else if ( CurrentStates[ButtonIndex] != 0 && ControllerState.NextRepeatTime[ButtonIndex] <= CurrentTime )
+				{
+					MessageHandler->OnControllerButtonPressed( Buttons[ButtonIndex], PlatformUser, InputDevice, true );
+
+					// set the button's NextRepeatTime to the ButtonRepeatDelay
+					ControllerState.NextRepeatTime[ButtonIndex] = CurrentTime + ButtonRepeatDelay;
+				}
+
+				// Update the state for next time
+				ControllerState.ButtonStates[ButtonIndex] = CurrentStates[ButtonIndex];
+			}	
+
+			// apply force feedback
+
+			const float LargeValue = (ControllerState.ForceFeedback.LeftLarge > ControllerState.ForceFeedback.RightLarge ? ControllerState.ForceFeedback.LeftLarge : ControllerState.ForceFeedback.RightLarge);
+			const float SmallValue = (ControllerState.ForceFeedback.LeftSmall > ControllerState.ForceFeedback.RightSmall ? ControllerState.ForceFeedback.LeftSmall : ControllerState.ForceFeedback.RightSmall);
+
+			if (!FMath::IsNearlyEqual(LargeValue, ControllerState.LastLargeValue) || !FMath::IsNearlyEqual(SmallValue, ControllerState.LastSmallValue))
+			{
+				XINPUT_VIBRATION VibrationState;
+				VibrationState.wLeftMotorSpeed = ( ::WORD ) ( LargeValue * 65535.0f );
+				VibrationState.wRightMotorSpeed = ( ::WORD ) ( SmallValue * 65535.0f );
+ 
+				XInputSetState( ( ::DWORD ) ControllerState.ControllerId, &VibrationState );
+
+				ControllerState.LastLargeValue = LargeValue;
+				ControllerState.LastSmallValue = SmallValue;
+			}
+		}
+	}
+
+	bNeedsControllerStateUpdate = false;
+}
+
+
+void XInputInterface::SetMessageHandler( const TSharedRef< FGenericApplicationMessageHandler >& InMessageHandler )
+{
+	MessageHandler = InMessageHandler;
+}
+
+void XInputInterface::SetDeviceProperty(int32 ControllerId, const FInputDeviceProperty* Property)
+{
+	// This should only get called from WindowsApplication.cpp when Windows detects a device change.
+	static const FName UpdateRequestedName = TEXT("Request_Device_Update");
+	if (Property && Property->Name == UpdateRequestedName)
+	{
+		SetNeedsControllerStateUpdate();	
+	}
+}
+
+void XInputInterface::SetChannelValue( int32 ControllerId, const FForceFeedbackChannelType ChannelType, const float Value )
+{
+	if (ControllerId >= 0 && ControllerId < MAX_NUM_XINPUT_CONTROLLERS)
+	{
+		FControllerState& ControllerState = ControllerStates[ ControllerId ];
+
+		if( ControllerState.bIsConnected )
+		{
+			switch( ChannelType )
+			{
+				case FForceFeedbackChannelType::LEFT_LARGE:
+					ControllerState.ForceFeedback.LeftLarge = Value;
+					break;
+
+				case FForceFeedbackChannelType::LEFT_SMALL:
+					ControllerState.ForceFeedback.LeftSmall = Value;
+					break;
+
+				case FForceFeedbackChannelType::RIGHT_LARGE:
+					ControllerState.ForceFeedback.RightLarge = Value;
+					break;
+
+				case FForceFeedbackChannelType::RIGHT_SMALL:
+					ControllerState.ForceFeedback.RightSmall = Value;
+					break;
+			}
+		}
+	}
+}
+
+void XInputInterface::SetChannelValues( int32 ControllerId, const FForceFeedbackValues &Values )
+{
+	if (ControllerId >= 0 && ControllerId < MAX_NUM_XINPUT_CONTROLLERS)
+	{
+		FControllerState& ControllerState = ControllerStates[ ControllerId ];
+
+		if( ControllerState.bIsConnected )
+		{
+			ControllerState.ForceFeedback = Values;
+		}
+	}
+}
diff --git a/Engine/Plugins/Runtime/Windows/XInputDevice/Source/XInputDevice/Private/XInputInterface.h b/Engine/Plugins/Runtime/Windows/XInputDevice/Source/XInputDevice/Private/XInputInterface.h
new file mode 100644
--- /dev/null
+++ b/Engine/Plugins/Runtime/Windows/XInputDevice/Source/XInputDevice/Private/XInputInterface.h
@@ -0,0 +1,136 @@
+// Copyright Epic Games, Inc. All Rights Reserved.
+
+#pragma once
+
+#include "CoreTypes.h"
+#include "GenericPlatform/IInputInterface.h"
+#include "IInputDevice.h"
+#include "GenericPlatform/GenericApplicationMessageHandler.h"
+
+/** Max number of controllers. */
+#define MAX_NUM_XINPUT_CONTROLLERS 4
+
+/** Max number of controller buttons.  Must be < 256*/
+#define MAX_NUM_CONTROLLER_BUTTONS 24
+
+enum class FForceFeedbackChannelType;
+
+/**
+ * Interface class for XInput devices (xbox 360 controller)                 
+ */
+class XInputInterface : public IInputDevice
+{
+public:
+
+	static TSharedRef< XInputInterface > Create( const TSharedRef< FGenericApplicationMessageHandler >& InMessageHandler );
+
+	/**
+	 * Poll for controller state and send events if needed
+	 *
+	 * @param PathToJoystickCaptureWidget	The path to the joystick capture widget.  If invalid this function does not poll 
+	 */
+	virtual void SendControllerEvents() override;
+
+	virtual void SetMessageHandler( const TSharedRef< FGenericApplicationMessageHandler >& InMessageHandler ) override;
+
+	void SetNeedsControllerStateUpdate() { bNeedsControllerStateUpdate = true; }
+
+	/**
+	* Sets the strength/speed of the given channel for the given controller id.
+	* NOTE: If the channel is not supported, the call will silently fail
+	*
+	* @param ControllerId the id of the controller whose value is to be set
+	* @param ChannelType the type of channel whose value should be set
+	* @param Value strength or speed of feedback, 0.0f to 1.0f. 0.0f will disable
+	*/
+	virtual void SetChannelValue( int32 ControllerId, const FForceFeedbackChannelType ChannelType, const float Value ) override;
+
+	/**
+	* Sets a property for a given controller id.
+	* Will be ignored for devices which don't support the property.
+	*
+	* @param ControllerId the id of the controller whose property is to be applied
+	* @param Property Base class pointer to property that will be applied
+	*/
+	virtual void SetDeviceProperty(int32 ControllerId, const FInputDeviceProperty* Property) override;
+
+	/**
+	* Sets the strength/speed of all the channels for the given controller id.
+	* NOTE: Unsupported channels are silently ignored
+	*
+	* @param ControllerId the id of the controller whose value is to be set
+	* @param Values strength or speed of feedback for all channels
+	*/
+	virtual void SetChannelValues( int32 ControllerId, const FForceFeedbackValues& Values ) override;
+
+	virtual bool IsGamepadAttached() const override { return bIsGamepadAttached; }
+	virtual void Tick( float DeltaTime ) override {};
+	virtual bool Exec( UWorld* InWorld, const TCHAR* Cmd, FOutputDevice& Ar ) { return false; }
+
+private:
+
+	XInputInterface( const TSharedRef< FGenericApplicationMessageHandler >& MessageHandler );
+
+
+	struct FControllerState
+	{
+		/** Last frame's button states, so we only send events on edges */
+		bool ButtonStates[MAX_NUM_CONTROLLER_BUTTONS];
+
+		/** Next time a repeat event should be generated for each button */
+		double NextRepeatTime[MAX_NUM_CONTROLLER_BUTTONS];
+
+		/** Raw Left thumb x analog value */
+		int16 LeftXAnalog;
+
+		/** Raw left thumb y analog value */
+		int16 LeftYAnalog;
+
+		/** Raw Right thumb x analog value */
+		int16 RightXAnalog;
+
+		/** Raw Right thumb x analog value */
+		int16 RightYAnalog;
+
+		/** Left Trigger analog value */
+		uint8 LeftTriggerAnalog;
+
+		/** Right trigger analog value */
+		uint8 RightTriggerAnalog;
+
+		/** If the controller is currently connected */
+		bool bIsConnected;
+
+		/** Id of the controller */
+		int32 ControllerId;
+	
+		/** Current force feedback values */
+		FForceFeedbackValues ForceFeedback;
+
+		float LastLargeValue;
+		float LastSmallValue;
+	};
+
+	/** If we've been notified by the system that the controller state may have changed */
+	bool bNeedsControllerStateUpdate;
+
+	bool bIsGamepadAttached;
+
+	/** In the engine, all controllers map to xbox controllers for consistency */
+	uint8	X360ToXboxControllerMapping[MAX_NUM_CONTROLLER_BUTTONS];
+
+	/** Controller states */
+	FControllerState ControllerStates[MAX_NUM_XINPUT_CONTROLLERS];
+
+	/** Delay before sending a repeat message after a button was first pressed */
+	float InitialButtonRepeatDelay;
+
+	/** Delay before sending a repeat message after a button has been pressed for a while */
+	float ButtonRepeatDelay;
+
+	/**  */
+	FGamepadKeyNames::Type Buttons[MAX_NUM_CONTROLLER_BUTTONS];
+
+	/**  */
+	TSharedRef<FGenericApplicationMessageHandler> MessageHandler;
+};
diff --git a/Engine/Plugins/Runtime/Windows/XInputDevice/Source/XInputDevice/Public/XInputDeviceModule.h b/Engine/Plugins/Runtime/Windows/XInputDevice/Source/XInputDevice/Public/XInputDeviceModule.h
new file mode 100644
--- /dev/null
+++ b/Engine/Plugins/Runtime/Windows/XInputDevice/Source/XInputDevice/Public/XInputDeviceModule.h
@@ -0,0 +1,16 @@
+// Copyright Epic Games, Inc. All Rights Reserved.
+
+#pragma once
+
+#include "IInputDeviceModule.h"
+
+/**
+ * Input Device Module for XInput, an input API for Windows.
+ */
+class FXInputDeviceModule final : public IInputDeviceModule
+{
+public:
+	//~ Begin IInputDeviceModule interface
+	virtual TSharedPtr<IInputDevice> CreateInputDevice(const TSharedRef<FGenericApplicationMessageHandler>& InMessageHandler) override;		
+	//~ End IInputDeviceModule interface
+};
\ No newline at end of file
diff --git a/Engine/Plugins/Runtime/Windows/XInputDevice/Source/XInputDevice/XInputDevice.build.cs b/Engine/Plugins/Runtime/Windows/XInputDevice/Source/XInputDevice/XInputDevice.build.cs
new file mode 100644
--- /dev/null
+++ b/Engine/Plugins/Runtime/Windows/XInputDevice/Source/XInputDevice/XInputDevice.build.cs
@@ -0,0 +1,37 @@
+// Copyright Epic Games, Inc. All Rights Reserved.
+
+using System;
+
+namespace UnrealBuildTool.Rules
+{
+	public class XInputDevice : ModuleRules
+	{
+		public XInputDevice(ReadOnlyTargetRules Target) : base(Target)
+		{
+			// Enable cast warnings as errors
+			UnsafeTypeCastWarningLevel = WarningLevel.Error;
+			
+			PublicDependencyModuleNames.AddRange(
+				new string[]
+				{
+					"Core",
+					"ApplicationCore",
+					"Engine",
+					"InputDevice",
+				}
+			);
+
+			// XInput is only supported on Windows
+			if (Target.IsInPlatformGroup(UnrealPlatformGroup.Windows))
+			{
+				AddEngineThirdPartyPrivateStaticDependencies(Target,
+					"XInput"
+				);
+			}
+			else
+			{
+				Console.Error.Write("XInput is only supported on the Windows target platform group!");
+			}
+		}
+	}
+}
\ No newline at end of file
diff --git a/Engine/Plugins/Runtime/Windows/XInputDevice/XInputDevice.uplugin b/Engine/Plugins/Runtime/Windows/XInputDevice/XInputDevice.uplugin
new file mode 100644
--- /dev/null
+++ b/Engine/Plugins/Runtime/Windows/XInputDevice/XInputDevice.uplugin
@@ -0,0 +1,30 @@
+{
+	"FileVersion": 3,
+	"FriendlyName": "XInput Device",
+	"Version": 1,
+	"VersionName": "1.0",
+	"CreatedBy": "Epic Games, Inc.",
+	"CreatedByURL": "https://epicgames.com",
+	"Description": "XInput is a Game Controller API for Windows.",
+	"Category": "Input Devices",
+	"EnabledByDefault": true,
+	"IsBetaVersion": false,
+	"CanContainContent": false,
+	"SupportedTargetPlatforms":
+	[
+		"Win64",
+		"WinGDK"
+	],
+	"Modules":
+	[
+		{
+			"Name": "XInputDevice",
+			"Type": "Runtime",
+			"PlatformAllowList":
+			[
+				"Win64",
+				"WinGDK"
+			]
+		}
+	]
+}
\ No newline at end of file
diff --git a/Engine/Source/Runtime/ApplicationCore/Private/Windows/WindowsApplication.cpp b/Engine/Source/Runtime/ApplicationCore/Private/Windows/WindowsApplication.cpp
--- a/Engine/Source/Runtime/ApplicationCore/Private/Windows/WindowsApplication.cpp
+++ b/Engine/Source/Runtime/ApplicationCore/Private/Windows/WindowsApplication.cpp
@@ -10,7 +10,6 @@
 #include "Misc/App.h"
 #include "Windows/WindowsWindow.h"
 #include "Windows/WindowsCursor.h"
-#include "XInputInterface.h"
 #include "Features/IModularFeatures.h"
 #include "IInputDeviceModule.h"
 #include "IInputDevice.h"
@@ -118,7 +117,6 @@ FWindowsApplication::FWindowsApplication( const HINSTANCE HInstance, const HICON
 	, bForceActivateByMouse( false )
 	, bForceNoGamepads( false )
 	, bConsumeAltSpace( false )
-	, XInput( XInputInterface::Create( MessageHandler ) )
 	, bHasLoadedInputPlugins( false )
 	, bAllowedToDeferMessageProcessing(true)
 	, CVarDeferMessageProcessing( 
@@ -384,7 +382,6 @@ void FWindowsApplication::InitializeWindow( const TSharedRef< FGenericWindow >&
 void FWindowsApplication::SetMessageHandler( const TSharedRef< FGenericApplicationMessageHandler >& InMessageHandler )
 {
 	GenericApplication::SetMessageHandler(InMessageHandler);
-	XInput->SetMessageHandler( InMessageHandler );
 
 	TArray<IInputDeviceModule*> PluginImplementations = IModularFeatures::Get().GetModularFeatureImplementations<IInputDeviceModule>( IInputDeviceModule::GetModularFeatureName() );
 	for( auto DeviceIt = ExternalInputDevices.CreateIterator(); DeviceIt; ++DeviceIt )
@@ -411,11 +408,6 @@ bool FWindowsApplication::IsGamepadAttached() const
 		return false;
 	}
 
-	if (XInput->IsGamepadAttached())
-	{
-		return true;
-	}
-
 	for( auto DeviceIt = ExternalInputDevices.CreateConstIterator(); DeviceIt; ++DeviceIt )
 	{
 		if ((*DeviceIt)->IsGamepadAttached())
@@ -1883,7 +1875,14 @@ int32 FWindowsApplication::ProcessMessage( HWND hwnd, uint32 msg, WPARAM wParam,
 
 		case WM_DEVICECHANGE:
 			{
-				XInput->SetNeedsControllerStateUpdate(); 
+				static const FInputDeviceProperty RequestUpdateProp(TEXT("Request_Device_Update"));
+				for( auto DeviceIt = ExternalInputDevices.CreateIterator(); DeviceIt; ++DeviceIt )
+				{
+					// Pass in -1 because this isn't for any specific device, it is just a message requesting an update to any attached controllers.
+					// At the moment, this will only be listened for by XInput.
+					(*DeviceIt)->SetDeviceProperty(-1, &RequestUpdateProp);
+				}
+				
 				QueryConnectedMice();
 			}
 			break;
@@ -2808,9 +2807,6 @@ void FWindowsApplication::PollGameDeviceState( const float TimeDelta )
 		return; // do not proceed if the app uses VR focus but doesn't have it
 	}
 
-	// Poll game device states and send new events
-	XInput->SendControllerEvents();
-
 	// Poll externally-implemented devices
 	for( auto DeviceIt = ExternalInputDevices.CreateIterator(); DeviceIt; ++DeviceIt )
 	{
@@ -2842,8 +2838,6 @@ void FWindowsApplication::SetForceFeedbackChannelValues(int32 ControllerId, cons
 
 	const FForceFeedbackValues* InternalValues = &Values;
  
-	XInput->SetChannelValues( ControllerId, *InternalValues );
- 
 	// send vibration to externally-implemented devices
 	for( auto DeviceIt = ExternalInputDevices.CreateIterator(); DeviceIt; ++DeviceIt )
 	{
diff --git a/Engine/Source/Runtime/ApplicationCore/Private/Windows/XInputInterface.cpp b/Engine/Source/Runtime/ApplicationCore/Private/Windows/XInputInterface.cpp
deleted file mode 100644
diff --git a/Engine/Source/Runtime/ApplicationCore/Private/Windows/XInputInterface.h b/Engine/Source/Runtime/ApplicationCore/Private/Windows/XInputInterface.h
deleted file mode 100644
diff --git a/Engine/Source/Runtime/ApplicationCore/Public/Windows/WindowsApplication.h b/Engine/Source/Runtime/ApplicationCore/Public/Windows/WindowsApplication.h
--- a/Engine/Source/Runtime/ApplicationCore/Public/Windows/WindowsApplication.h
+++ b/Engine/Source/Runtime/ApplicationCore/Public/Windows/WindowsApplication.h
@@ -480,8 +480,6 @@ private:
 
 	TArray<TSharedRef<FWindowsWindow>> Windows;
 
-	TSharedRef<class XInputInterface> XInput;
-
 	/** List of input devices implemented in external modules. */
 	TArray<TSharedPtr<class IInputDevice>> ExternalInputDevices;
 	bool bHasLoadedInputPlugins;
