Split windows platform extensions into their own child plugins

#rnx

diff --git a/Engine/Plugins/Runtime/Windows/XInputDevice/XInputDevice.uplugin b/Engine/Plugins/Runtime/Windows/XInputDevice/XInputDevice.uplugin
--- a/Engine/Plugins/Runtime/Windows/XInputDevice/XInputDevice.uplugin
+++ b/Engine/Plugins/Runtime/Windows/XInputDevice/XInputDevice.uplugin
@@ -12,8 +12,7 @@
 	"CanContainContent": false,
 	"SupportedTargetPlatforms":
 	[
-		"Win64",
-		"WinGDK"
+		"Win64"
 	],
 	"Modules":
 	[
@@ -22,8 +21,7 @@
 			"Type": "Runtime",
 			"PlatformAllowList":
 			[
-				"Win64",
-				"WinGDK"
+				"Win64"
 			]
 		}
 	]
