Add support for creating secondary input devices to be used outside of the Unreal Input System.

diff --git a/Engine/Plugins/Runtime/Windows/XInputDevice/Source/XInputDevice/Private/XInputDeviceModule.cpp b/Engine/Plugins/Runtime/Windows/XInputDevice/Source/XInputDevice/Private/XInputDeviceModule.cpp
--- a/Engine/Plugins/Runtime/Windows/XInputDevice/Source/XInputDevice/Private/XInputDeviceModule.cpp
+++ b/Engine/Plugins/Runtime/Windows/XInputDevice/Source/XInputDevice/Private/XInputDeviceModule.cpp
@@ -6,7 +6,13 @@
 
 TSharedPtr<IInputDevice> FXInputDeviceModule::CreateInputDevice(const TSharedRef<FGenericApplicationMessageHandler>& InMessageHandler)
 {
-	return XInputInterface::Create(InMessageHandler);
+	return XInputInterface::Create(InMessageHandler, true /* bShouldBePrimaryDevice */);
 }
 
-IMPLEMENT_MODULE(FXInputDeviceModule, XInputDevice)
\ No newline at end of file
+TSharedPtr<IInputDevice> FXInputDeviceModule::CreateInputDevice(const TSharedRef<FGenericApplicationMessageHandler>& InMessageHandler, FInputDeviceCreationParameters InParameters)
+{
+	// XInputDevice can run on any thread and supports secondary device creation.
+	return XInputInterface::Create(InMessageHandler, InParameters.bInitAsPrimaryDevice);
+}
+
+IMPLEMENT_MODULE(FXInputDeviceModule, XInputDevice)
diff --git a/Engine/Plugins/Runtime/Windows/XInputDevice/Source/XInputDevice/Private/XInputInterface.cpp b/Engine/Plugins/Runtime/Windows/XInputDevice/Source/XInputDevice/Private/XInputInterface.cpp
--- a/Engine/Plugins/Runtime/Windows/XInputDevice/Source/XInputDevice/Private/XInputInterface.cpp
+++ b/Engine/Plugins/Runtime/Windows/XInputDevice/Source/XInputDevice/Private/XInputInterface.cpp
@@ -3,6 +3,7 @@
 #include "XInputInterface.h"
 #include "HAL/PlatformTime.h"
 #include "Misc/CoreDelegates.h"
+#include "Misc/CoreMiscDefines.h"
 #include "Windows/WindowsApplication.h"
 #include "GenericPlatform/GenericApplication.h"
 #include "GenericPlatform/GenericPlatformInputDeviceMapper.h"
@@ -22,14 +23,14 @@ FAutoConsoleVariableRef CVarForceControllerStateUpdate(
 	TEXT("0: Not Enabled, 1: Enabled"),
 	ECVF_Default);
 
-TSharedRef< XInputInterface > XInputInterface::Create(  const TSharedRef< FGenericApplicationMessageHandler >& InMessageHandler )
+TSharedRef< XInputInterface > XInputInterface::Create(  const TSharedRef< FGenericApplicationMessageHandler >& InMessageHandler, bool bShouldBePrimaryDevice )
 {
-	return MakeShareable( new XInputInterface( InMessageHandler ) );
+	return MakeShareable( new XInputInterface( InMessageHandler, bShouldBePrimaryDevice ) );
 }
 
 
-XInputInterface::XInputInterface( const TSharedRef< FGenericApplicationMessageHandler >& InMessageHandler )
-	: MessageHandler( InMessageHandler )
+XInputInterface::XInputInterface(const TSharedRef<FGenericApplicationMessageHandler>& InMessageHandler, bool bShouldBePrimaryDevice)
+	: bIsPrimaryDevice(bShouldBePrimaryDevice), MessageHandler(InMessageHandler)
 {
 	for ( int32 ControllerIndex=0; ControllerIndex < MAX_NUM_XINPUT_CONTROLLERS; ++ControllerIndex )
 	{
@@ -109,6 +110,42 @@ float ShortToNormalizedFloat(int16 AxisVal)
 static FName XInputInterfaceName = FName("XInputInterface");
 static FString XInputControllerIdentifier = TEXT("XInputController");
 
+void XInputInterface::GetPlatformUserAndDevice(int32 InControllerId, EInputDeviceConnectionState InDeviceState, FPlatformUserId& OutPlatformUserId, FInputDeviceId& OutDeviceId)
+{
+	if (bIsPrimaryDevice)
+	{
+		IPlatformInputDeviceMapper& DeviceMapper = IPlatformInputDeviceMapper::Get();
+		DeviceMapper.RemapControllerIdToPlatformUserAndDevice(InControllerId, OUT OutPlatformUserId, OUT OutDeviceId);
+
+		// If the controller is connected now but was not before, refresh the information
+		if (InDeviceState == EInputDeviceConnectionState::Connected || InDeviceState == EInputDeviceConnectionState::Disconnected)
+		{
+			DeviceMapper.Internal_MapInputDeviceToUser(OutDeviceId, OutPlatformUserId, InDeviceState);
+		}
+	}
+	else
+	{
+		// Use the controller id as the device id for secondary input devices not connected to the input system.
+		OutDeviceId = FInputDeviceId::CreateFromInternalId(InControllerId);
+	}
+}
+
+namespace UE::XInputInterface::Private
+{
+EInputDeviceConnectionState GetInputDeviceConnectionState(bool bWasConnected, bool bControllerStateIsConnected)
+{
+	if (!bWasConnected && bControllerStateIsConnected)
+	{
+		return EInputDeviceConnectionState::Connected;
+	}
+	else if (bWasConnected && !bControllerStateIsConnected)
+	{
+		return EInputDeviceConnectionState::Disconnected;
+	}
+	return EInputDeviceConnectionState::Unknown;
+}
+
+}
 void XInputInterface::SendControllerEvents()
 {
 	bool bWereConnected[MAX_NUM_XINPUT_CONTROLLERS];
@@ -150,21 +187,11 @@ void XInputInterface::SendControllerEvents()
 		{
 			const XINPUT_STATE& XInputState = XInputStates[ControllerIndex];
 
-			IPlatformInputDeviceMapper& DeviceMapper = IPlatformInputDeviceMapper::Get();
 			FPlatformUserId PlatformUser = PLATFORMUSERID_NONE;
 			FInputDeviceId InputDevice = INPUTDEVICEID_NONE;
-			DeviceMapper.RemapControllerIdToPlatformUserAndDevice(ControllerState.ControllerId, OUT PlatformUser, OUT InputDevice);
+			EInputDeviceConnectionState State = UE::XInputInterface::Private::GetInputDeviceConnectionState(bWasConnected, ControllerState.bIsConnected);
+			GetPlatformUserAndDevice(ControllerState.ControllerId, State, OUT PlatformUser, OUT InputDevice);
 
-			// If the controller is connected now but was not before, refresh the information
-			if (!bWasConnected && ControllerState.bIsConnected)
-			{
-				DeviceMapper.Internal_MapInputDeviceToUser(InputDevice, PlatformUser, EInputDeviceConnectionState::Connected);
-			}
-			else if (bWasConnected && !ControllerState.bIsConnected)
-			{
-				DeviceMapper.Internal_MapInputDeviceToUser(InputDevice, PlatformUser, EInputDeviceConnectionState::Disconnected);
-			}
-			
 			bool CurrentStates[MAX_NUM_CONTROLLER_BUTTONS] = {0};
 		
 			// Get the current state of all buttons
diff --git a/Engine/Plugins/Runtime/Windows/XInputDevice/Source/XInputDevice/Private/XInputInterface.h b/Engine/Plugins/Runtime/Windows/XInputDevice/Source/XInputDevice/Private/XInputInterface.h
--- a/Engine/Plugins/Runtime/Windows/XInputDevice/Source/XInputDevice/Private/XInputInterface.h
+++ b/Engine/Plugins/Runtime/Windows/XInputDevice/Source/XInputDevice/Private/XInputInterface.h
@@ -6,6 +6,7 @@
 #include "GenericPlatform/IInputInterface.h"
 #include "IInputDevice.h"
 #include "GenericPlatform/GenericApplicationMessageHandler.h"
+#include "Misc/CoreMiscDefines.h"
 
 /** Max number of controllers. */
 #define MAX_NUM_XINPUT_CONTROLLERS 4
@@ -22,7 +23,7 @@ class XInputInterface : public IInputDevice
 {
 public:
 
-	static TSharedRef< XInputInterface > Create( const TSharedRef< FGenericApplicationMessageHandler >& InMessageHandler );
+	static TSharedRef< XInputInterface > Create( const TSharedRef< FGenericApplicationMessageHandler >& InMessageHandler, bool bShouldBePrimaryDevice );
 
 	/**
 	 * Poll for controller state and send events if needed
@@ -69,8 +70,12 @@ public:
 
 private:
 
-	XInputInterface( const TSharedRef< FGenericApplicationMessageHandler >& MessageHandler );
+	XInputInterface( const TSharedRef< FGenericApplicationMessageHandler >& MessageHandler, bool bShouldBePrimaryDevice );
 
+	/**
+	 * Maps the given controller id to the platform user and device id. It will use the device mapper when used with the input system.
+	 */
+	void GetPlatformUserAndDevice( int32 InControllerId, EInputDeviceConnectionState InDeviceState, FPlatformUserId& OutPlatformUserId, FInputDeviceId& OutDeviceId );
 
 	struct FControllerState
 	{
@@ -116,6 +121,9 @@ private:
 
 	bool bIsGamepadAttached;
 
+	/** Indicates if this device is operating as a primary device and thus part of game input system. */
+	bool bIsPrimaryDevice;
+
 	/** In the engine, all controllers map to xbox controllers for consistency */
 	uint8	X360ToXboxControllerMapping[MAX_NUM_CONTROLLER_BUTTONS];
 
diff --git a/Engine/Plugins/Runtime/Windows/XInputDevice/Source/XInputDevice/Public/XInputDeviceModule.h b/Engine/Plugins/Runtime/Windows/XInputDevice/Source/XInputDevice/Public/XInputDeviceModule.h
--- a/Engine/Plugins/Runtime/Windows/XInputDevice/Source/XInputDevice/Public/XInputDeviceModule.h
+++ b/Engine/Plugins/Runtime/Windows/XInputDevice/Source/XInputDevice/Public/XInputDeviceModule.h
@@ -11,6 +11,7 @@ class FXInputDeviceModule final : public IInputDeviceModule
 {
 public:
 	//~ Begin IInputDeviceModule interface
-	virtual TSharedPtr<IInputDevice> CreateInputDevice(const TSharedRef<FGenericApplicationMessageHandler>& InMessageHandler) override;		
+	virtual TSharedPtr<IInputDevice> CreateInputDevice(const TSharedRef<FGenericApplicationMessageHandler>& InMessageHandler) override;
+	virtual TSharedPtr<IInputDevice> CreateInputDevice(const TSharedRef<FGenericApplicationMessageHandler>& InMessageHandler, FInputDeviceCreationParameters InParameters) override;
 	//~ End IInputDeviceModule interface
-};
\ No newline at end of file
+};
diff --git a/Engine/Source/Runtime/ApplicationCore/Private/GenericPlatform/GenericApplication.cpp b/Engine/Source/Runtime/ApplicationCore/Private/GenericPlatform/GenericApplication.cpp
--- a/Engine/Source/Runtime/ApplicationCore/Private/GenericPlatform/GenericApplication.cpp
+++ b/Engine/Source/Runtime/ApplicationCore/Private/GenericPlatform/GenericApplication.cpp
@@ -45,7 +45,14 @@ const FGamepadKeyNames::Type FGamepadKeyNames::RightStickDown("Gamepad_RightStic
 const FGamepadKeyNames::Type FGamepadKeyNames::RightStickRight("Gamepad_RightStick_Right");
 const FGamepadKeyNames::Type FGamepadKeyNames::RightStickLeft("Gamepad_RightStick_Left");
 
-TArray<FInputDeviceScope*> FInputDeviceScope::ScopeStack;
+namespace UE::InputDeviceScope::Private
+{
+TArray<FInputDeviceScope*>& GetScopeStack()
+{
+	static thread_local TArray<FInputDeviceScope*> ScopeStack;
+	return ScopeStack;
+}
+}
 
 FInputDeviceScope::FInputDeviceScope(IInputDevice* InInputDevice, FName InInputDeviceName, int32 InHardwareDeviceHandle, FString InHardwareDeviceIdentifier)
 	: InputDevice(InInputDevice)
@@ -53,31 +60,25 @@ FInputDeviceScope::FInputDeviceScope(IInputDevice* InInputDevice, FName InInputD
 	, HardwareDeviceHandle(InHardwareDeviceHandle)
 	, HardwareDeviceIdentifier(InHardwareDeviceIdentifier)
 {
-	if (ensure(IsInGameThread()))
-	{
-		// Add to scope stack
-		ScopeStack.Add(this);
-	}
+	// Add to scope stack
+	UE::InputDeviceScope::Private::GetScopeStack().Add(this);
 }
 
 FInputDeviceScope::~FInputDeviceScope()
 {
-	if (ensure(IsInGameThread()))
-	{
-		// This should always be the top of the stack
-		ensureMsgf((ScopeStack.Num() > 0 && ScopeStack.Last() == this), TEXT("FInputDeviceScope was not destroyed in correct order!"));
-		ScopeStack.Remove(this);
-	}
+	TArray<FInputDeviceScope*>& ScopeStack = UE::InputDeviceScope::Private::GetScopeStack();
+
+	// This should always be the top of the stack
+	ensureMsgf((ScopeStack.Num() > 0 && ScopeStack.Last() == this), TEXT("FInputDeviceScope was not destroyed in correct order!"));
+	ScopeStack.Remove(this);
 }
 
 const FInputDeviceScope* FInputDeviceScope::GetCurrent()
 {
-	if (ensure(IsInGameThread()))
+	TArray<FInputDeviceScope*>& ScopeStack = UE::InputDeviceScope::Private::GetScopeStack();
+	if (ScopeStack.Num() > 0)
 	{
-		if (ScopeStack.Num() > 0)
-		{
-			return ScopeStack.Last();
-		}
+		return ScopeStack.Last();
 	}
 	return nullptr;
 }
diff --git a/Engine/Source/Runtime/ApplicationCore/Public/GenericPlatform/GenericApplicationMessageHandler.h b/Engine/Source/Runtime/ApplicationCore/Public/GenericPlatform/GenericApplicationMessageHandler.h
--- a/Engine/Source/Runtime/ApplicationCore/Public/GenericPlatform/GenericApplicationMessageHandler.h
+++ b/Engine/Source/Runtime/ApplicationCore/Public/GenericPlatform/GenericApplicationMessageHandler.h
@@ -201,9 +201,6 @@ public:
 
 	/** Returns the currently active InputDeviceScope. This is only valid to call on the game thread and may return null */
 	static APPLICATIONCORE_API const FInputDeviceScope* GetCurrent();
-
-private:
-	static APPLICATIONCORE_API TArray<FInputDeviceScope*> ScopeStack;
 };
 
 PRAGMA_DISABLE_DEPRECATION_WARNINGS
diff --git a/Engine/Source/Runtime/InputDevice/Public/IInputDeviceModule.h b/Engine/Source/Runtime/InputDevice/Public/IInputDeviceModule.h
--- a/Engine/Source/Runtime/InputDevice/Public/IInputDeviceModule.h
+++ b/Engine/Source/Runtime/InputDevice/Public/IInputDeviceModule.h
@@ -10,6 +10,13 @@
 #include "Modules/ModuleManager.h"
 #include "IInputDevice.h"
 
+/** Input parameters to device creation. */
+struct INPUTDEVICE_API FInputDeviceCreationParameters
+{
+	/** Indicates if this device is operating as a primary device and thus a part of game input system. */
+	bool bInitAsPrimaryDevice = true;
+};
+
 /**
  * The public interface of the InputDeviceModule
  */
@@ -48,10 +55,29 @@ public:
 	}
 
 	/**
-	 * Attempts to create a new head tracking device interface
+	 * Attempts to create a new input device with the given message handler.  This version will eventually be deprecated and users should
+	 * migrate to the new version with an additional parameters struct.
 	 *
-	 * @return	Interface to the new head tracking device, if we were able to successfully create one
+	 * @return	Interface to the new input device, if we were able to successfully create one
 	 */
-	virtual TSharedPtr< class IInputDevice > CreateInputDevice(const TSharedRef< FGenericApplicationMessageHandler >& InMessageHandler) = 0;
+	virtual TSharedPtr<class IInputDevice> CreateInputDevice(const TSharedRef<FGenericApplicationMessageHandler>& InMessageHandler) = 0;
 
+	/**
+	 * Attempts to create a new input device interface with the specified input
+	 * device creation parameters. This override is for advanced use cases where
+	 * users would like to create devices that are not a part of the input
+	 * system.
+	 *
+	 * @return	Interface to the new input device, if we were able to successfully create one
+	 */
+	virtual TSharedPtr<class IInputDevice> CreateInputDevice(const TSharedRef<FGenericApplicationMessageHandler>& InMessageHandler, FInputDeviceCreationParameters InParams)
+	{
+		// If this is a primary device and it should be constructed for the input system then call the default implementation.
+		if (InParams.bInitAsPrimaryDevice)
+		{
+			return CreateInputDevice(InMessageHandler);
+		}
+
+		return nullptr;
+	}
 };
