Fix Live Link Input device not working in 5.6 main.

diff --git a/Engine/Plugins/Runtime/Windows/XInputDevice/Source/XInputDevice/Private/XInputInterface.cpp b/Engine/Plugins/Runtime/Windows/XInputDevice/Source/XInputDevice/Private/XInputInterface.cpp
--- a/Engine/Plugins/Runtime/Windows/XInputDevice/Source/XInputDevice/Private/XInputInterface.cpp
+++ b/Engine/Plugins/Runtime/Windows/XInputDevice/Source/XInputDevice/Private/XInputInterface.cpp
@@ -221,7 +221,7 @@ void XInputInterface::SendControllerEvents()
 			}
 
 			// If we never got a valid platform user for some reason, we can't do anything
-			if (!PlatformUser.IsValid())
+			if (!PlatformUser.IsValid() && bIsPrimaryDevice)
 			{
 				return;
 			}
