PR #14044: Check for latest device independent of type so other device types don't trigger feedback on the gamepad.

#rnx
#Synced-CL 48850057

diff --git a/Engine/Plugins/Runtime/Windows/XInputDevice/Source/XInputDevice/Private/XInputInterface.cpp b/Engine/Plugins/Runtime/Windows/XInputDevice/Source/XInputDevice/Private/XInputInterface.cpp
--- a/Engine/Plugins/Runtime/Windows/XInputDevice/Source/XInputDevice/Private/XInputInterface.cpp
+++ b/Engine/Plugins/Runtime/Windows/XInputDevice/Source/XInputDevice/Private/XInputInterface.cpp
@@ -424,8 +424,8 @@ void XInputInterface::SetChannelValue( int32 ControllerId, const FForceFeedbackC
 	UInputDeviceSubsystem* DeviceSubsystem = UInputDeviceSubsystem::Get();
 	check(DeviceSubsystem);
 	
-	// Get the latest gamepad input device
-	const FInputDeviceId MostRecentDevice = DeviceSubsystem->GetLatestDeviceOfType(UserId, EHardwareDevicePrimaryType::Gamepad);
+	// Get the latest input device
+	const FInputDeviceId MostRecentDevice = DeviceSubsystem->GetMostRecentlyUsedInputDeviceId(UserId);
 	
 	if (ControllerId >= 0 && ControllerId < MAX_NUM_XINPUT_CONTROLLERS)
 	{
@@ -468,7 +468,8 @@ void XInputInterface::SetChannelValues( int32 ControllerId, const FForceFeedback
 
 	UInputDeviceSubsystem* DeviceSubsystem = UInputDeviceSubsystem::Get();
 	check(DeviceSubsystem);
-	const FInputDeviceId MostRecentDevice = DeviceSubsystem->GetLatestDeviceOfType(UserId, EHardwareDevicePrimaryType::Gamepad);
+	// Get the latest input device
+	const FInputDeviceId MostRecentDevice = DeviceSubsystem->GetMostRecentlyUsedInputDeviceId(UserId);
 	
 	if (ControllerId >= 0 && ControllerId < MAX_NUM_XINPUT_CONTROLLERS)
 	{
