[Xinput] add a comment to easily disable optimization in the XInput build module

#rnx
#Synced-CL 52309720

diff --git a/Engine/Plugins/Runtime/Windows/XInputDevice/Source/XInputDevice/XInputDevice.build.cs b/Engine/Plugins/Runtime/Windows/XInputDevice/Source/XInputDevice/XInputDevice.build.cs
--- a/Engine/Plugins/Runtime/Windows/XInputDevice/Source/XInputDevice/XInputDevice.build.cs
+++ b/Engine/Plugins/Runtime/Windows/XInputDevice/Source/XInputDevice/XInputDevice.build.cs
@@ -8,6 +8,9 @@ namespace UnrealBuildTool.Rules
 	{
 		public XInputDevice(ReadOnlyTargetRules Target) : base(Target)
 		{
+			// Uncomment this line to make for easier debugging
+			//OptimizeCode = CodeOptimization.Never;
+						
 			// Enable cast warnings as errors
 			CppCompileWarningSettings.UnsafeTypeCastWarningLevel = WarningLevel.Error;
 			
